Updated Plugin Browser Templates to use Tool Menus

diff --git a/Engine/Plugins/Editor/PluginBrowser/Templates/Advanced/Source/PLUGIN_NAME/PLUGIN_NAME.Build.cs b/Engine/Plugins/Editor/PluginBrowser/Templates/Advanced/Source/PLUGIN_NAME/PLUGIN_NAME.Build.cs
--- a/Engine/Plugins/Editor/PluginBrowser/Templates/Advanced/Source/PLUGIN_NAME/PLUGIN_NAME.Build.cs
+++ b/Engine/Plugins/Editor/PluginBrowser/Templates/Advanced/Source/PLUGIN_NAME/PLUGIN_NAME.Build.cs
@@ -37,7 +37,7 @@ public class PLUGIN_NAME : ModuleRules
 				"Projects",
 				"InputCore",
 				"UnrealEd",
-				"LevelEditor",
+				"ToolMenus",
 				"CoreUObject",
 				"Engine",
 				"Slate",
diff --git a/Engine/Plugins/Editor/PluginBrowser/Templates/Advanced/Source/PLUGIN_NAME/Private/PLUGIN_NAME.cpp b/Engine/Plugins/Editor/PluginBrowser/Templates/Advanced/Source/PLUGIN_NAME/Private/PLUGIN_NAME.cpp
--- a/Engine/Plugins/Editor/PluginBrowser/Templates/Advanced/Source/PLUGIN_NAME/Private/PLUGIN_NAME.cpp
+++ b/Engine/Plugins/Editor/PluginBrowser/Templates/Advanced/Source/PLUGIN_NAME/Private/PLUGIN_NAME.cpp
@@ -7,7 +7,7 @@
 #include "Widgets/Docking/SDockTab.h"
 #include "Widgets/Layout/SBox.h"
 #include "Widgets/Text/STextBlock.h"
-#include "Framework/MultiBox/MultiBoxBuilder.h"
+#include "ToolMenus.h"
 
 static const FName PLUGIN_NAMETabName("PLUGIN_NAME");
 
@@ -28,22 +28,8 @@ void FPLUGIN_NAMEModule::StartupModule()
 		FPLUGIN_NAMECommands::Get().OpenPluginWindow,
 		FExecuteAction::CreateRaw(this, &FPLUGIN_NAMEModule::PluginButtonClicked),
 		FCanExecuteAction());
-		
-	FLevelEditorModule& LevelEditorModule = FModuleManager::LoadModuleChecked<FLevelEditorModule>("LevelEditor");
-	
-	{
-		TSharedPtr<FExtender> MenuExtender = MakeShareable(new FExtender());
-		MenuExtender->AddMenuExtension("WindowLayout", EExtensionHook::After, PluginCommands, FMenuExtensionDelegate::CreateRaw(this, &FPLUGIN_NAMEModule::AddMenuExtension));
 
-		LevelEditorModule.GetMenuExtensibilityManager()->AddExtender(MenuExtender);
-	}
-	
-	{
-		TSharedPtr<FExtender> ToolbarExtender = MakeShareable(new FExtender);
-		ToolbarExtender->AddToolBarExtension("Settings", EExtensionHook::After, PluginCommands, FToolBarExtensionDelegate::CreateRaw(this, &FPLUGIN_NAMEModule::AddToolbarExtension));
-		
-		LevelEditorModule.GetToolBarExtensibilityManager()->AddExtender(ToolbarExtender);
-	}
+	UToolMenus::RegisterStartupCallback(FSimpleMulticastDelegate::FDelegate::CreateRaw(this, &FPLUGIN_NAMEModule::RegisterMenus));
 	
 	FGlobalTabmanager::Get()->RegisterNomadTabSpawner(PLUGIN_NAMETabName, FOnSpawnTab::CreateRaw(this, &FPLUGIN_NAMEModule::OnSpawnPluginTab))
 		.SetDisplayName(LOCTEXT("FPLUGIN_NAMETabTitle", "PLUGIN_NAME"))
@@ -54,6 +40,11 @@ void FPLUGIN_NAMEModule::ShutdownModule()
 {
 	// This function may be called during shutdown to clean up your module.  For modules that support dynamic reloading,
 	// we call this function before unloading the module.
+
+	UToolMenus::UnRegisterStartupCallback(this);
+
+	UToolMenus::UnregisterOwner(this);
+
 	FPLUGIN_NAMEStyle::Shutdown();
 
 	FPLUGIN_NAMECommands::Unregister();
@@ -88,14 +79,29 @@ void FPLUGIN_NAMEModule::PluginButtonClicked()
 	FGlobalTabmanager::Get()->InvokeTab(PLUGIN_NAMETabName);
 }
 
-void FPLUGIN_NAMEModule::AddMenuExtension(FMenuBuilder& Builder)
+void FPLUGIN_NAMEModule::RegisterMenus()
 {
-	Builder.AddMenuEntry(FPLUGIN_NAMECommands::Get().OpenPluginWindow);
-}
+	// Owner will be used for cleanup in call to UToolMenus::UnregisterOwner
+	FToolMenuOwnerScoped OwnerScoped(this);
 
-void FPLUGIN_NAMEModule::AddToolbarExtension(FToolBarBuilder& Builder)
-{
-	Builder.AddToolBarButton(FPLUGIN_NAMECommands::Get().OpenPluginWindow);
+	{
+		UToolMenu* Menu = UToolMenus::Get()->ExtendMenu("LevelEditor.MainMenu.Window");
+		{
+			FToolMenuSection& Section = Menu->FindOrAddSection("WindowLayout");
+			Section.AddMenuEntryWithCommandList(FPLUGIN_NAMECommands::Get().OpenPluginWindow, PluginCommands);
+		}
+	}
+
+	{
+		UToolMenu* ToolbarMenu = UToolMenus::Get()->ExtendMenu("LevelEditor.LevelEditorToolBar");
+		{
+			FToolMenuSection& Section = ToolbarMenu->FindOrAddSection("Settings");
+			{
+				FToolMenuEntry& Entry = Section.AddEntry(FToolMenuEntry::InitToolBarButton(FPLUGIN_NAMECommands::Get().OpenPluginWindow));
+				Entry.SetCommandList(PluginCommands);
+			}
+		}
+	}
 }
 
 #undef LOCTEXT_NAMESPACE
diff --git a/Engine/Plugins/Editor/PluginBrowser/Templates/Advanced/Source/PLUGIN_NAME/Public/PLUGIN_NAME.h b/Engine/Plugins/Editor/PluginBrowser/Templates/Advanced/Source/PLUGIN_NAME/Public/PLUGIN_NAME.h
--- a/Engine/Plugins/Editor/PluginBrowser/Templates/Advanced/Source/PLUGIN_NAME/Public/PLUGIN_NAME.h
+++ b/Engine/Plugins/Editor/PluginBrowser/Templates/Advanced/Source/PLUGIN_NAME/Public/PLUGIN_NAME.h
@@ -21,8 +21,7 @@ public:
 	
 private:
 
-	void AddToolbarExtension(FToolBarBuilder& Builder);
-	void AddMenuExtension(FMenuBuilder& Builder);
+	void RegisterMenus();
 
 	TSharedRef<class SDockTab> OnSpawnPluginTab(const class FSpawnTabArgs& SpawnTabArgs);
 
diff --git a/Engine/Plugins/Editor/PluginBrowser/Templates/Basic/Source/PLUGIN_NAME/PLUGIN_NAME.Build.cs b/Engine/Plugins/Editor/PluginBrowser/Templates/Basic/Source/PLUGIN_NAME/PLUGIN_NAME.Build.cs
--- a/Engine/Plugins/Editor/PluginBrowser/Templates/Basic/Source/PLUGIN_NAME/PLUGIN_NAME.Build.cs
+++ b/Engine/Plugins/Editor/PluginBrowser/Templates/Basic/Source/PLUGIN_NAME/PLUGIN_NAME.Build.cs
@@ -37,7 +37,7 @@ public class PLUGIN_NAME : ModuleRules
 				"Projects",
 				"InputCore",
 				"UnrealEd",
-				"LevelEditor",
+				"ToolMenus",
 				"CoreUObject",
 				"Engine",
 				"Slate",
diff --git a/Engine/Plugins/Editor/PluginBrowser/Templates/Basic/Source/PLUGIN_NAME/Private/PLUGIN_NAME.cpp b/Engine/Plugins/Editor/PluginBrowser/Templates/Basic/Source/PLUGIN_NAME/Private/PLUGIN_NAME.cpp
--- a/Engine/Plugins/Editor/PluginBrowser/Templates/Basic/Source/PLUGIN_NAME/Private/PLUGIN_NAME.cpp
+++ b/Engine/Plugins/Editor/PluginBrowser/Templates/Basic/Source/PLUGIN_NAME/Private/PLUGIN_NAME.cpp
@@ -4,9 +4,7 @@
 #include "PLUGIN_NAMEStyle.h"
 #include "PLUGIN_NAMECommands.h"
 #include "Misc/MessageDialog.h"
-#include "Framework/MultiBox/MultiBoxBuilder.h"
-
-#include "LevelEditor.h"
+#include "ToolMenus.h"
 
 static const FName PLUGIN_NAMETabName("PLUGIN_NAME");
 
@@ -27,28 +25,19 @@ void FPLUGIN_NAMEModule::StartupModule()
 		FPLUGIN_NAMECommands::Get().PluginAction,
 		FExecuteAction::CreateRaw(this, &FPLUGIN_NAMEModule::PluginButtonClicked),
 		FCanExecuteAction());
-		
-	FLevelEditorModule& LevelEditorModule = FModuleManager::LoadModuleChecked<FLevelEditorModule>("LevelEditor");
-	
-	{
-		TSharedPtr<FExtender> MenuExtender = MakeShareable(new FExtender());
-		MenuExtender->AddMenuExtension("WindowLayout", EExtensionHook::After, PluginCommands, FMenuExtensionDelegate::CreateRaw(this, &FPLUGIN_NAMEModule::AddMenuExtension));
 
-		LevelEditorModule.GetMenuExtensibilityManager()->AddExtender(MenuExtender);
-	}
-	
-	{
-		TSharedPtr<FExtender> ToolbarExtender = MakeShareable(new FExtender);
-		ToolbarExtender->AddToolBarExtension("Settings", EExtensionHook::After, PluginCommands, FToolBarExtensionDelegate::CreateRaw(this, &FPLUGIN_NAMEModule::AddToolbarExtension));
-		
-		LevelEditorModule.GetToolBarExtensibilityManager()->AddExtender(ToolbarExtender);
-	}
+	UToolMenus::RegisterStartupCallback(FSimpleMulticastDelegate::FDelegate::CreateRaw(this, &FPLUGIN_NAMEModule::RegisterMenus));
 }
 
 void FPLUGIN_NAMEModule::ShutdownModule()
 {
 	// This function may be called during shutdown to clean up your module.  For modules that support dynamic reloading,
 	// we call this function before unloading the module.
+
+	UToolMenus::UnRegisterStartupCallback(this);
+
+	UToolMenus::UnregisterOwner(this);
+
 	FPLUGIN_NAMEStyle::Shutdown();
 
 	FPLUGIN_NAMECommands::Unregister();
@@ -65,14 +54,29 @@ void FPLUGIN_NAMEModule::PluginButtonClicked()
 	FMessageDialog::Open(EAppMsgType::Ok, DialogText);
 }
 
-void FPLUGIN_NAMEModule::AddMenuExtension(FMenuBuilder& Builder)
+void FPLUGIN_NAMEModule::RegisterMenus()
 {
-	Builder.AddMenuEntry(FPLUGIN_NAMECommands::Get().PluginAction);
-}
+	// Owner will be used for cleanup in call to UToolMenus::UnregisterOwner
+	FToolMenuOwnerScoped OwnerScoped(this);
 
-void FPLUGIN_NAMEModule::AddToolbarExtension(FToolBarBuilder& Builder)
-{
-	Builder.AddToolBarButton(FPLUGIN_NAMECommands::Get().PluginAction);
+	{
+		UToolMenu* Menu = UToolMenus::Get()->ExtendMenu("LevelEditor.MainMenu.Window");
+		{
+			FToolMenuSection& Section = Menu->FindOrAddSection("WindowLayout");
+			Section.AddMenuEntryWithCommandList(FPLUGIN_NAMECommands::Get().PluginAction, PluginCommands);
+		}
+	}
+
+	{
+		UToolMenu* ToolbarMenu = UToolMenus::Get()->ExtendMenu("LevelEditor.LevelEditorToolBar");
+		{
+			FToolMenuSection& Section = ToolbarMenu->FindOrAddSection("Settings");
+			{
+				FToolMenuEntry& Entry = Section.AddEntry(FToolMenuEntry::InitToolBarButton(FPLUGIN_NAMECommands::Get().PluginAction));
+				Entry.SetCommandList(PluginCommands);
+			}
+		}
+	}
 }
 
 #undef LOCTEXT_NAMESPACE
diff --git a/Engine/Plugins/Editor/PluginBrowser/Templates/Basic/Source/PLUGIN_NAME/Public/PLUGIN_NAME.h b/Engine/Plugins/Editor/PluginBrowser/Templates/Basic/Source/PLUGIN_NAME/Public/PLUGIN_NAME.h
--- a/Engine/Plugins/Editor/PluginBrowser/Templates/Basic/Source/PLUGIN_NAME/Public/PLUGIN_NAME.h
+++ b/Engine/Plugins/Editor/PluginBrowser/Templates/Basic/Source/PLUGIN_NAME/Public/PLUGIN_NAME.h
@@ -21,8 +21,8 @@ public:
 	
 private:
 
-	void AddToolbarExtension(FToolBarBuilder& Builder);
-	void AddMenuExtension(FMenuBuilder& Builder);
+	void RegisterMenus();
+
 
 private:
 	TSharedPtr<class FUICommandList> PluginCommands;
