Fixed use of deprecated FInputGesture in some plugin wizard templates (and two use cases generated from those templates)

diff --git a/Engine/Plugins/Editor/PluginBrowser/Templates/Advanced/Source/PLUGIN_NAME/Private/PLUGIN_NAMECommands.cpp b/Engine/Plugins/Editor/PluginBrowser/Templates/Advanced/Source/PLUGIN_NAME/Private/PLUGIN_NAMECommands.cpp
--- a/Engine/Plugins/Editor/PluginBrowser/Templates/Advanced/Source/PLUGIN_NAME/Private/PLUGIN_NAMECommands.cpp
+++ b/Engine/Plugins/Editor/PluginBrowser/Templates/Advanced/Source/PLUGIN_NAME/Private/PLUGIN_NAMECommands.cpp
@@ -6,7 +6,7 @@
 
 void FPLUGIN_NAMECommands::RegisterCommands()
 {
-	UI_COMMAND(OpenPluginWindow, "PLUGIN_NAME", "Bring up PLUGIN_NAME window", EUserInterfaceActionType::Button, FInputGesture());
+	UI_COMMAND(OpenPluginWindow, "PLUGIN_NAME", "Bring up PLUGIN_NAME window", EUserInterfaceActionType::Button, FInputChord());
 }
 
 #undef LOCTEXT_NAMESPACE
diff --git a/Engine/Plugins/Editor/PluginBrowser/Templates/Basic/Source/PLUGIN_NAME/Private/PLUGIN_NAMECommands.cpp b/Engine/Plugins/Editor/PluginBrowser/Templates/Basic/Source/PLUGIN_NAME/Private/PLUGIN_NAMECommands.cpp
--- a/Engine/Plugins/Editor/PluginBrowser/Templates/Basic/Source/PLUGIN_NAME/Private/PLUGIN_NAMECommands.cpp
+++ b/Engine/Plugins/Editor/PluginBrowser/Templates/Basic/Source/PLUGIN_NAME/Private/PLUGIN_NAMECommands.cpp
@@ -6,7 +6,7 @@
 
 void FPLUGIN_NAMECommands::RegisterCommands()
 {
-	UI_COMMAND(PluginAction, "PLUGIN_NAME", "Execute PLUGIN_NAME action", EUserInterfaceActionType::Button, FInputGesture());
+	UI_COMMAND(PluginAction, "PLUGIN_NAME", "Execute PLUGIN_NAME action", EUserInterfaceActionType::Button, FInputChord());
 }
 
 #undef LOCTEXT_NAMESPACE
