Plugin browser fix templates to use the new suggested slate style and icons. Fix button look.

diff --git a/Engine/Plugins/Editor/PluginBrowser/Source/PluginBrowser/PluginBrowser.Build.cs b/Engine/Plugins/Editor/PluginBrowser/Source/PluginBrowser/PluginBrowser.Build.cs
--- a/Engine/Plugins/Editor/PluginBrowser/Source/PluginBrowser/PluginBrowser.Build.cs
+++ b/Engine/Plugins/Editor/PluginBrowser/Source/PluginBrowser/PluginBrowser.Build.cs
@@ -31,7 +31,8 @@ namespace UnrealBuildTool.Rules
 					"MainFrame",
                     "UATHelper",
                     "AssetTools",
-					"Json"
+					"Json",
+					"ToolWidgets",
 				}
 			);
 
diff --git a/Engine/Plugins/Editor/PluginBrowser/Source/PluginBrowser/Private/SNewPluginWizard.cpp b/Engine/Plugins/Editor/PluginBrowser/Source/PluginBrowser/Private/SNewPluginWizard.cpp
--- a/Engine/Plugins/Editor/PluginBrowser/Source/PluginBrowser/Private/SNewPluginWizard.cpp
+++ b/Engine/Plugins/Editor/PluginBrowser/Source/PluginBrowser/Private/SNewPluginWizard.cpp
@@ -44,6 +44,7 @@
 #include "Framework/Notifications/NotificationManager.h"
 #include "Widgets/Notifications/SNotificationList.h"
 #include "SourceCodeNavigation.h"
+#include "SPrimaryButton.h"
 
 DEFINE_LOG_CATEGORY(LogPluginWizard);
 
@@ -261,15 +262,11 @@ void SNewPluginWizard::Construct(const FArguments& Args, TSharedPtr<SDockTab> In
 
 	MainContent->AddSlot()
 	.AutoHeight()
-	.Padding(5)
+	.Padding(10)
 	.HAlign(HAlign_Right)
 	[
-		SNew(SButton)
-		.ContentPadding(5)
-		.TextStyle(FEditorStyle::Get(), "LargeText")
-		.ButtonStyle(FEditorStyle::Get(), "FlatButton.Success")
+		SNew(SPrimaryButton)
 		.IsEnabled(this, &SNewPluginWizard::CanCreatePlugin)
-		.HAlign(HAlign_Center)
 		.Text(CreateButtonLabel)
 		.OnClicked(this, &SNewPluginWizard::OnCreatePluginClicked)
 	];
diff --git a/Engine/Plugins/Editor/PluginBrowser/Source/PluginBrowser/Private/SPluginBrowser.cpp b/Engine/Plugins/Editor/PluginBrowser/Source/PluginBrowser/Private/SPluginBrowser.cpp
--- a/Engine/Plugins/Editor/PluginBrowser/Source/PluginBrowser/Private/SPluginBrowser.cpp
+++ b/Engine/Plugins/Editor/PluginBrowser/Source/PluginBrowser/Private/SPluginBrowser.cpp
@@ -23,6 +23,7 @@
 #include "ProjectDescriptor.h"
 #include "GameProjectGenerationModule.h"
 #include "Framework/MultiBox/MultiBoxBuilder.h"
+#include "SPrimaryButton.h"
 
 #define LOCTEXT_NAMESPACE "PluginsEditor"
 
@@ -153,8 +154,7 @@ void SPluginBrowser::Construct( const FArguments& Args )
 		]
 		+SSplitter::Slot()
 		[
-			SNew( SVerticalBox )
-
+			SNew(SVerticalBox)
 			+SVerticalBox::Slot()
 			.Padding(FMargin(0.0f, 0.0f, 0.0f, PaddingAmount))
 			.AutoHeight()
@@ -247,16 +247,11 @@ void SPluginBrowser::Construct( const FArguments& Args )
 
 			+ SVerticalBox::Slot()
 			.AutoHeight()
-			.Padding(FMargin(PaddingAmount, PaddingAmount, PaddingAmount, 0.0f))
+			.Padding(FMargin(2.0f, 10.0f, 10.0f, 10.0f))
 			.HAlign(HAlign_Right)
 			[
-				SNew(SButton)
-				.ContentPadding(5)
-				.IsEnabled(true)
+				SNew(SPrimaryButton)
 				.ToolTip(SNew(SToolTip).Text(LOCTEXT("NewPluginEnabled", "Click here to open the Plugin Creator dialog.")))
-				.TextStyle(FEditorStyle::Get(), "LargeText")
-				.ButtonStyle(FEditorStyle::Get(), "FlatButton.Success")
-				.HAlign(HAlign_Center)
 				.Text(LOCTEXT("NewPluginLabel", "New Plugin"))
 				.OnClicked(this, &SPluginBrowser::HandleNewPluginButtonClicked)
 			]
diff --git a/Engine/Plugins/Editor/PluginBrowser/Templates/Advanced/Source/PLUGIN_NAME/Private/PLUGIN_NAMEStyle.cpp b/Engine/Plugins/Editor/PluginBrowser/Templates/Advanced/Source/PLUGIN_NAME/Private/PLUGIN_NAMEStyle.cpp
--- a/Engine/Plugins/Editor/PluginBrowser/Templates/Advanced/Source/PLUGIN_NAME/Private/PLUGIN_NAMEStyle.cpp
+++ b/Engine/Plugins/Editor/PluginBrowser/Templates/Advanced/Source/PLUGIN_NAME/Private/PLUGIN_NAMEStyle.cpp
@@ -5,8 +5,11 @@
 #include "Framework/Application/SlateApplication.h"
 #include "Slate/SlateGameResources.h"
 #include "Interfaces/IPluginManager.h"
+#include "Styling/SlateStyleMacros.h"
 
-TSharedPtr< FSlateStyleSet > FPLUGIN_NAMEStyle::StyleInstance = NULL;
+#define RootToContentDir Style->RootToContentDir
+
+TSharedPtr<FSlateStyleSet> FPLUGIN_NAMEStyle::StyleInstance = nullptr;
 
 void FPLUGIN_NAMEStyle::Initialize()
 {
@@ -30,32 +33,19 @@ FName FPLUGIN_NAMEStyle::GetStyleSetName()
 	return StyleSetName;
 }
 
-#define IMAGE_BRUSH( RelativePath, ... ) FSlateImageBrush( Style->RootToContentDir( RelativePath, TEXT(".png") ), __VA_ARGS__ )
-#define BOX_BRUSH( RelativePath, ... ) FSlateBoxBrush( Style->RootToContentDir( RelativePath, TEXT(".png") ), __VA_ARGS__ )
-#define BORDER_BRUSH( RelativePath, ... ) FSlateBorderBrush( Style->RootToContentDir( RelativePath, TEXT(".png") ), __VA_ARGS__ )
-#define TTF_FONT( RelativePath, ... ) FSlateFontInfo( Style->RootToContentDir( RelativePath, TEXT(".ttf") ), __VA_ARGS__ )
-#define OTF_FONT( RelativePath, ... ) FSlateFontInfo( Style->RootToContentDir( RelativePath, TEXT(".otf") ), __VA_ARGS__ )
-
 const FVector2D Icon16x16(16.0f, 16.0f);
 const FVector2D Icon20x20(20.0f, 20.0f);
-const FVector2D Icon40x40(40.0f, 40.0f);
 
 TSharedRef< FSlateStyleSet > FPLUGIN_NAMEStyle::Create()
 {
 	TSharedRef< FSlateStyleSet > Style = MakeShareable(new FSlateStyleSet("PLUGIN_NAMEStyle"));
 	Style->SetContentRoot(IPluginManager::Get().FindPlugin("PLUGIN_NAME")->GetBaseDir() / TEXT("Resources"));
 
-	Style->Set("PLUGIN_NAME.OpenPluginWindow", new IMAGE_BRUSH(TEXT("ButtonIcon_40x"), Icon40x40));
+	Style->Set("PLUGIN_NAME.OpenPluginWindow", new IMAGE_BRUSH_SVG(TEXT("PlaceholderButtonIcon"), Icon20x20));
 
 	return Style;
 }
 
-#undef IMAGE_BRUSH
-#undef BOX_BRUSH
-#undef BORDER_BRUSH
-#undef TTF_FONT
-#undef OTF_FONT
-
 void FPLUGIN_NAMEStyle::ReloadTextures()
 {
 	if (FSlateApplication::IsInitialized())
diff --git a/Engine/Plugins/Editor/PluginBrowser/Templates/Basic/Source/PLUGIN_NAME/Private/PLUGIN_NAME.cpp b/Engine/Plugins/Editor/PluginBrowser/Templates/Basic/Source/PLUGIN_NAME/Private/PLUGIN_NAME.cpp
--- a/Engine/Plugins/Editor/PluginBrowser/Templates/Basic/Source/PLUGIN_NAME/Private/PLUGIN_NAME.cpp
+++ b/Engine/Plugins/Editor/PluginBrowser/Templates/Basic/Source/PLUGIN_NAME/Private/PLUGIN_NAME.cpp
@@ -68,9 +68,9 @@ void FPLUGIN_NAMEModule::RegisterMenus()
 	}
 
 	{
-		UToolMenu* ToolbarMenu = UToolMenus::Get()->ExtendMenu("LevelEditor.LevelEditorToolBar");
+		UToolMenu* ToolbarMenu = UToolMenus::Get()->ExtendMenu("LevelEditor.LevelEditorToolBar.PlayToolBar");
 		{
-			FToolMenuSection& Section = ToolbarMenu->FindOrAddSection("Settings");
+			FToolMenuSection& Section = ToolbarMenu->FindOrAddSection("PluginTools");
 			{
 				FToolMenuEntry& Entry = Section.AddEntry(FToolMenuEntry::InitToolBarButton(FPLUGIN_NAMECommands::Get().PluginAction));
 				Entry.SetCommandList(PluginCommands);
diff --git a/Engine/Plugins/Editor/PluginBrowser/Templates/Basic/Source/PLUGIN_NAME/Private/PLUGIN_NAMEStyle.cpp b/Engine/Plugins/Editor/PluginBrowser/Templates/Basic/Source/PLUGIN_NAME/Private/PLUGIN_NAMEStyle.cpp
--- a/Engine/Plugins/Editor/PluginBrowser/Templates/Basic/Source/PLUGIN_NAME/Private/PLUGIN_NAMEStyle.cpp
+++ b/Engine/Plugins/Editor/PluginBrowser/Templates/Basic/Source/PLUGIN_NAME/Private/PLUGIN_NAMEStyle.cpp
@@ -6,8 +6,11 @@
 #include "Styling/SlateStyleRegistry.h"
 #include "Slate/SlateGameResources.h"
 #include "Interfaces/IPluginManager.h"
+#include "Styling/SlateStyleMacros.h"
 
-TSharedPtr< FSlateStyleSet > FPLUGIN_NAMEStyle::StyleInstance = NULL;
+#define RootToContentDir Style->RootToContentDir
+
+TSharedPtr<FSlateStyleSet> FPLUGIN_NAMEStyle::StyleInstance = nullptr;
 
 void FPLUGIN_NAMEStyle::Initialize()
 {
@@ -31,32 +34,19 @@ FName FPLUGIN_NAMEStyle::GetStyleSetName()
 	return StyleSetName;
 }
 
-#define IMAGE_BRUSH( RelativePath, ... ) FSlateImageBrush( Style->RootToContentDir( RelativePath, TEXT(".png") ), __VA_ARGS__ )
-#define BOX_BRUSH( RelativePath, ... ) FSlateBoxBrush( Style->RootToContentDir( RelativePath, TEXT(".png") ), __VA_ARGS__ )
-#define BORDER_BRUSH( RelativePath, ... ) FSlateBorderBrush( Style->RootToContentDir( RelativePath, TEXT(".png") ), __VA_ARGS__ )
-#define TTF_FONT( RelativePath, ... ) FSlateFontInfo( Style->RootToContentDir( RelativePath, TEXT(".ttf") ), __VA_ARGS__ )
-#define OTF_FONT( RelativePath, ... ) FSlateFontInfo( Style->RootToContentDir( RelativePath, TEXT(".otf") ), __VA_ARGS__ )
 
 const FVector2D Icon16x16(16.0f, 16.0f);
 const FVector2D Icon20x20(20.0f, 20.0f);
-const FVector2D Icon40x40(40.0f, 40.0f);
 
 TSharedRef< FSlateStyleSet > FPLUGIN_NAMEStyle::Create()
 {
 	TSharedRef< FSlateStyleSet > Style = MakeShareable(new FSlateStyleSet("PLUGIN_NAMEStyle"));
 	Style->SetContentRoot(IPluginManager::Get().FindPlugin("PLUGIN_NAME")->GetBaseDir() / TEXT("Resources"));
 
-	Style->Set("PLUGIN_NAME.PluginAction", new IMAGE_BRUSH(TEXT("ButtonIcon_40x"), Icon40x40));
-
+	Style->Set("PLUGIN_NAME.PluginAction", new IMAGE_BRUSH_SVG(TEXT("PlaceholderButtonIcon"), Icon20x20));
 	return Style;
 }
 
-#undef IMAGE_BRUSH
-#undef BOX_BRUSH
-#undef BORDER_BRUSH
-#undef TTF_FONT
-#undef OTF_FONT
-
 void FPLUGIN_NAMEStyle::ReloadTextures()
 {
 	if (FSlateApplication::IsInitialized())
