PR #9208: Fix Advanced Plugin Template to extend correct PlayToolBar menu

diff --git a/Engine/Plugins/Editor/PluginBrowser/Templates/Advanced/Source/PLUGIN_NAME/Private/PLUGIN_NAME.cpp b/Engine/Plugins/Editor/PluginBrowser/Templates/Advanced/Source/PLUGIN_NAME/Private/PLUGIN_NAME.cpp
--- a/Engine/Plugins/Editor/PluginBrowser/Templates/Advanced/Source/PLUGIN_NAME/Private/PLUGIN_NAME.cpp
+++ b/Engine/Plugins/Editor/PluginBrowser/Templates/Advanced/Source/PLUGIN_NAME/Private/PLUGIN_NAME.cpp
@@ -93,9 +93,9 @@ void FPLUGIN_NAMEModule::RegisterMenus()
 	}
 
 	{
-		UToolMenu* ToolbarMenu = UToolMenus::Get()->ExtendMenu("LevelEditor.LevelEditorToolBar");
+		UToolMenu* ToolbarMenu = UToolMenus::Get()->ExtendMenu("LevelEditor.LevelEditorToolBar.PlayToolBar");
 		{
-			FToolMenuSection& Section = ToolbarMenu->FindOrAddSection("Settings");
+			FToolMenuSection& Section = ToolbarMenu->FindOrAddSection("PluginTools");
 			{
 				FToolMenuEntry& Entry = Section.AddEntry(FToolMenuEntry::InitToolBarButton(FPLUGIN_NAMECommands::Get().OpenPluginWindow));
 				Entry.SetCommandList(PluginCommands);
